MotorSimOutputMotoSynth - experimental plugin that implements the MotorSim output interface on a MotoSynth component

diff --git a/Engine/Plugins/Experimental/MotorSimOutputMotoSynth/MotorSimOutputMotoSynth.uplugin b/Engine/Plugins/Experimental/MotorSimOutputMotoSynth/MotorSimOutputMotoSynth.uplugin
new file mode 100644
--- /dev/null
+++ b/Engine/Plugins/Experimental/MotorSimOutputMotoSynth/MotorSimOutputMotoSynth.uplugin
@@ -0,0 +1,36 @@
+{
+	"FileVersion" : 1,
+	"Version" : 1,
+	"VersionName" : "1",
+	"FriendlyName" : "MotorSimOutputMotoSynth",
+	"Description" : "A MotoSim Output component using MotoSynth.",
+	"Category" : "Audio",
+	"CreatedBy" : "Epic Games, Inc.",
+	"CreatedByURL" : "http://epicgames.com",
+	"DocsURL" : "",
+	"MarketplaceURL" : "",
+	"SupportURL" : "",
+	"EnabledByDefault" : true,
+	"CanContainContent" : true,
+	"IsBetaVersion" : true,
+	"Installed" : false,
+	"Modules" :
+	[
+		{
+			"Name" : "MotorSimOutputMotoSynth",
+			"Type" : "Runtime",
+			"LoadingPhase" : "Default"
+		}
+	],
+	"Plugins" :
+		[
+			{
+				"Name" : "AudioMotorSim",
+				"Enabled": true
+			},
+			{
+				"Name" : "MotoSynth",
+				"Enabled": true
+			}
+		]
+}
\ No newline at end of file
diff --git a/Engine/Plugins/Experimental/MotorSimOutputMotoSynth/Source/MotorSimOutputMotoSynth/MotorSimOutputMotoSynth.build.cs b/Engine/Plugins/Experimental/MotorSimOutputMotoSynth/Source/MotorSimOutputMotoSynth/MotorSimOutputMotoSynth.build.cs
new file mode 100644
--- /dev/null
+++ b/Engine/Plugins/Experimental/MotorSimOutputMotoSynth/Source/MotorSimOutputMotoSynth/MotorSimOutputMotoSynth.build.cs
@@ -0,0 +1,28 @@
+// Copyright Epic Games, Inc. All Rights Reserved.
+
+namespace UnrealBuildTool.Rules
+{
+	public class MotorSimOutputMotoSynth : ModuleRules
+	{
+		public MotorSimOutputMotoSynth(ReadOnlyTargetRules Target) : base(Target)
+		{
+			PublicDependencyModuleNames.AddRange(
+				new string[] {
+					"Core",
+					"CoreUObject",
+					"Engine",
+					"AudioExtensions",
+					"MotoSynth",
+				}
+			);
+
+			PrivateDependencyModuleNames.AddRange(
+				new string[] {
+					"AudioMixer",
+					"AudioMotorSim",
+					"SignalProcessing",
+				}
+			);
+		}
+	}
+}
\ No newline at end of file
diff --git a/Engine/Plugins/Experimental/MotorSimOutputMotoSynth/Source/MotorSimOutputMotoSynth/Private/MotorSimOutputMotoSynth.cpp b/Engine/Plugins/Experimental/MotorSimOutputMotoSynth/Source/MotorSimOutputMotoSynth/Private/MotorSimOutputMotoSynth.cpp
new file mode 100644
--- /dev/null
+++ b/Engine/Plugins/Experimental/MotorSimOutputMotoSynth/Source/MotorSimOutputMotoSynth/Private/MotorSimOutputMotoSynth.cpp
@@ -0,0 +1,28 @@
+// Copyright Epic Games, Inc. All Rights Reserved.
+
+#pragma once
+
+#include "MotorSimOutputMotoSynth.h"
+#include "AudioMotorSimTypes.h"
+
+void UMotorSimOutputMotoSynth::Update(FAudioMotorSimInputContext& Input, FAudioMotorSimRuntimeContext& RuntimeInfo)
+{
+	float MinRpm, MaxRpm;
+	GetRPMRange(MinRpm, MaxRpm);
+
+	if (!FMath::IsNearlyEqual(MinRpm, MaxRpm, KINDA_SMALL_NUMBER))
+	{
+		const float RpmLog = Audio::GetLogFrequencyClamped(RuntimeInfo.Rpm, { 0.0f, 1.0f }, { MinRpm, MaxRpm });
+
+		SetRPM(RpmLog, Input.DeltaTime);
+	}
+}
+
+void UMotorSimOutputMotoSynth::StartOutput()
+{
+	Start();
+}
+void UMotorSimOutputMotoSynth::StopOutput()
+{
+	Stop();
+}
diff --git a/Engine/Plugins/Experimental/MotorSimOutputMotoSynth/Source/MotorSimOutputMotoSynth/Private/MotorSimOutputMotoSynthModule.cpp b/Engine/Plugins/Experimental/MotorSimOutputMotoSynth/Source/MotorSimOutputMotoSynth/Private/MotorSimOutputMotoSynthModule.cpp
new file mode 100644
--- /dev/null
+++ b/Engine/Plugins/Experimental/MotorSimOutputMotoSynth/Source/MotorSimOutputMotoSynth/Private/MotorSimOutputMotoSynthModule.cpp
@@ -0,0 +1,19 @@
+// Copyright Epic Games, Inc. All Rights Reserved.
+
+#include "MotorSimOutputMotoSynthModule.h"
+#include "CoreMinimal.h"
+#include "Modules/ModuleInterface.h"
+#include "Modules/ModuleManager.h"
+
+IMPLEMENT_MODULE(FMotorSimOutputMotoSynthModule, MotorSimOutputMotoSynth)
+
+void FMotorSimOutputMotoSynthModule::StartupModule()
+{
+	FModuleManager::Get().LoadModuleChecked(TEXT("SignalProcessing"));
+}
+
+void FMotorSimOutputMotoSynthModule::ShutdownModule()
+{
+}
+
+
diff --git a/Engine/Plugins/Experimental/MotorSimOutputMotoSynth/Source/MotorSimOutputMotoSynth/Public/MotorSimOutputMotoSynth.h b/Engine/Plugins/Experimental/MotorSimOutputMotoSynth/Source/MotorSimOutputMotoSynth/Public/MotorSimOutputMotoSynth.h
new file mode 100644
--- /dev/null
+++ b/Engine/Plugins/Experimental/MotorSimOutputMotoSynth/Source/MotorSimOutputMotoSynth/Public/MotorSimOutputMotoSynth.h
@@ -0,0 +1,22 @@
+// Copyright Epic Games, Inc. All Rights Reserved.
+
+#pragma once
+
+#include "AudioMotorSimTypes.h"
+#include "IAudioMotorSimOutput.h"
+#include "SynthComponents/SynthComponentMoto.h"
+#include "MotorSimOutputMotoSynth.generated.h"
+
+UCLASS(ClassGroup = Synth, meta = (BlueprintSpawnableComponent))
+class MOTORSIMOUTPUTMOTOSYNTH_API	UMotorSimOutputMotoSynth : public USynthComponentMoto, public IAudioMotorSimOutput
+{
+	GENERATED_BODY()
+
+	virtual ~UMotorSimOutputMotoSynth() override {};
+	
+public:
+	virtual void Update(FAudioMotorSimInputContext& Input, FAudioMotorSimRuntimeContext& RuntimeInfo) override;
+
+	virtual void StartOutput() override;
+	virtual void StopOutput() override;
+};
\ No newline at end of file
diff --git a/Engine/Plugins/Experimental/MotorSimOutputMotoSynth/Source/MotorSimOutputMotoSynth/Public/MotorSimOutputMotoSynthModule.h b/Engine/Plugins/Experimental/MotorSimOutputMotoSynth/Source/MotorSimOutputMotoSynth/Public/MotorSimOutputMotoSynthModule.h
new file mode 100644
--- /dev/null
+++ b/Engine/Plugins/Experimental/MotorSimOutputMotoSynth/Source/MotorSimOutputMotoSynth/Public/MotorSimOutputMotoSynthModule.h
@@ -0,0 +1,18 @@
+// Copyright Epic Games, Inc. All Rights Reserved.
+
+#pragma once
+
+#include "Modules/ModuleInterface.h"
+#include "Modules/ModuleManager.h"
+
+class FMotorSimOutputMotoSynthModule : public IModuleInterface
+{
+public:
+
+	// IModuleInterface interface
+	virtual void StartupModule() override;
+	virtual void ShutdownModule() override;
+};
+
+
+
