Ensure that AudioMotorSim and MotorSimOutputMotoSynth plugins are configured correctly for experimental public release

diff --git a/Engine/Plugins/Experimental/MotorSimOutputMotoSynth/MotorSimOutputMotoSynth.uplugin b/Engine/Plugins/Experimental/MotorSimOutputMotoSynth/MotorSimOutputMotoSynth.uplugin
--- a/Engine/Plugins/Experimental/MotorSimOutputMotoSynth/MotorSimOutputMotoSynth.uplugin
+++ b/Engine/Plugins/Experimental/MotorSimOutputMotoSynth/MotorSimOutputMotoSynth.uplugin
@@ -3,18 +3,17 @@
 	"Version" : 1,
 	"VersionName" : "1",
 	"FriendlyName" : "MotorSimOutputMotoSynth",
-	"Description" : "A MotoSim Output component using MotoSynth.",
+	"Description" : "A MotorSim Output component using MotoSynth.",
 	"Category" : "Audio",
 	"CreatedBy" : "Epic Games, Inc.",
 	"CreatedByURL" : "http://epicgames.com",
 	"DocsURL" : "",
 	"MarketplaceURL" : "",
 	"SupportURL" : "",
-	"EnabledByDefault" : true,
 	"CanContainContent" : true,
-	"IsBetaVersion" : true,
+	"IsExperimentalVersion": true,
 	"Installed" : false,
-	"Modules" :
+	"Modules" : 
 	[
 		{
 			"Name" : "MotorSimOutputMotoSynth",
diff --git a/Engine/Plugins/Runtime/AudioMotorSim/AudioMotorSim.uplugin b/Engine/Plugins/Runtime/AudioMotorSim/AudioMotorSim.uplugin
--- a/Engine/Plugins/Runtime/AudioMotorSim/AudioMotorSim.uplugin
+++ b/Engine/Plugins/Runtime/AudioMotorSim/AudioMotorSim.uplugin
@@ -2,7 +2,7 @@
 	"FileVersion": 1,
 	"Version": 1,
 	"VersionName": "1.0",
-	"FriendlyName": "AudioWidgets",
+	"FriendlyName": "Audio Motor Sim",
 	"Description": "Compositional method for simulating audio for vehicles.",
 	"Category": "Audio",
 	"CreatedBy": "Epic Games, Inc.",
@@ -11,9 +11,9 @@
 	"MarketplaceURL": "",
 	"SupportURL": "",
 	"CanContainContent": false,
-	"IsBetaVersion": true,
 	"IsExperimentalVersion": true,
 	"Installed": false,
+	"EnabledByDefault" : false,
 	"Modules": [
 		{
 			"Name": "AudioMotorSim",
